Updated DMXBP TemplateDefs.ini to update the Chinese description for this template

diff --git a/Templates/TP_DMXBP/Config/TemplateDefs.ini b/Templates/TP_DMXBP/Config/TemplateDefs.ini
--- a/Templates/TP_DMXBP/Config/TemplateDefs.ini
+++ b/Templates/TP_DMXBP/Config/TemplateDefs.ini
@@ -12,7 +12,7 @@ LocalizedDisplayNames=(Language="ja",Text="DMX")
 LocalizedDescriptions=(Language="ja",Text="このテンプレートは DMX プロトコルを必要とするライブイベントのプリビズとバーチャルプロダクション ワークフローのためのものです。プロキシ照明装置のアドレス、パッチ、コントロール、ライブ DMX データストリームを UE から入出力して記録、再生する例などを含んでいます。\n\nプロフェッショナル向け照明コンソールを使ったライブショーのプリビズ制作から、外部の実際の DMX 対応デバイスの Unreal からの操作、プロジェクトから DMX を使ったカスタムイベントのトリガーといった用途で使用できます。実験的なピクセルマッピング機能と、パイロテクニクス、花火、噴水などのエフェクトの追加コントロールについても例が用意されています。 \n\n")
 
 LocalizedDisplayNames=(Language="zh-Hans",Text="DMX")
-LocalizedDescriptions=(Language="zh-Hans",Text="该模板专门用于现场活动的视效预览和需要DMX协议的虚拟制片流程。它包含了寻址、配接和控制代理照明设备，以及在UE内外录制和回放现场DMX数据流的示例。\n\n它是一个很好的起点，可让你为使用专业照明控制台的现场演出制作视效预览。它可以让你在虚幻内控制外部启用了DMX的实体设备，或是在你使用了DMX的项目内部触发自定义事件。它还展示了实验性的像素映射功能和额外的可控制效果，如冷焰火、烟花和喷泉。\n\n")
+LocalizedDescriptions=(Language="zh-Hans",Text="该模板专门用于现场活动的视效预览和需要DMX协议的虚拟制片流程。它包含了寻址、配接和控制代理照明设备，以及在UE内外录制和回放现场DMX数据流的示例。\n\n该DMX模板是一个很好的起点，你可以为使用专业照明控制台的现场演出制作视效预览，可以在虚幻内控制启用了DMX的外部实体设备，或在使用了DMX的项目内部触发自定义事件。该模板还展示了实验性的像素映射功能和额外的可控制效果，如冷焰火、烟花和喷泉。\n\n")
 
 ; This project uses a thumbnail as the project icon
 bThumbnailAsIcon=false
